Fixed compile error in shipping config due to missing header include

diff --git a/Engine/Source/Runtime/CrunchCompression/Private/CrunchCompression.cpp b/Engine/Source/Runtime/CrunchCompression/Private/CrunchCompression.cpp
--- a/Engine/Source/Runtime/CrunchCompression/Private/CrunchCompression.cpp
+++ b/Engine/Source/Runtime/CrunchCompression/Private/CrunchCompression.cpp
@@ -6,6 +6,7 @@
 #include "PixelFormat.h"
 #include "Misc/ScopeLock.h"
 #include "HAL/IConsoleManager.h"
+#include "Stats/Stats.h"
 #include "Stats/Stats2.h"
 
 const bool GAdaptiveBlockSizes = true;
