Remove un-needed dependencies from CrunchCompression module, prevent circular dependency chain which causes problems on certain platforms

diff --git a/Engine/Source/Runtime/CrunchCompression/CrunchCompression.build.cs b/Engine/Source/Runtime/CrunchCompression/CrunchCompression.build.cs
--- a/Engine/Source/Runtime/CrunchCompression/CrunchCompression.build.cs
+++ b/Engine/Source/Runtime/CrunchCompression/CrunchCompression.build.cs
@@ -16,18 +16,6 @@ namespace UnrealBuildTool.Rules
 					"Crunch"
 				}
 			);
-
-            if (Target.Type == TargetType.Editor)
-            {
-                PublicDependencyModuleNames.AddRange(
-                    new string[]
-                    {
-                        "ImageCore",
-                        "TextureCompressor"
-                    }
-                 );
-            }
-
         }
 	}
 }
diff --git a/Engine/Source/Runtime/CrunchCompression/Private/CrunchCompression.cpp b/Engine/Source/Runtime/CrunchCompression/Private/CrunchCompression.cpp
--- a/Engine/Source/Runtime/CrunchCompression/Private/CrunchCompression.cpp
+++ b/Engine/Source/Runtime/CrunchCompression/Private/CrunchCompression.cpp
@@ -3,8 +3,6 @@
 
 #include "CrunchCompression.h"
 #include "Modules/ModuleManager.h"
-#include "PixelFormat.h"
-#include "Misc/ScopeLock.h"
 #include "HAL/IConsoleManager.h"
 #include "Stats/Stats.h"
 #include "Stats/Stats2.h"
diff --git a/Engine/Source/Runtime/CrunchCompression/Public/CrunchCompression.h b/Engine/Source/Runtime/CrunchCompression/Public/CrunchCompression.h
--- a/Engine/Source/Runtime/CrunchCompression/Public/CrunchCompression.h
+++ b/Engine/Source/Runtime/CrunchCompression/Public/CrunchCompression.h
@@ -4,10 +4,6 @@
 
 #include "CoreMinimal.h"
 
-#if WITH_EDITOR
-#include "ImageCore.h"
-#endif
-
 #ifndef WITH_CRUNCH
 #define WITH_CRUNCH 1
 #endif
