Use proper guard macro (WITH_CRUNCH_COMPRESSION) for initialing crunch compression allocator

diff --git a/Engine/Source/Runtime/CrunchCompression/Private/CrunchCompression.cpp b/Engine/Source/Runtime/CrunchCompression/Private/CrunchCompression.cpp
--- a/Engine/Source/Runtime/CrunchCompression/Private/CrunchCompression.cpp
+++ b/Engine/Source/Runtime/CrunchCompression/Private/CrunchCompression.cpp
@@ -182,9 +182,9 @@ struct FCrunchRegisterAllocators
 {
 	FCrunchRegisterAllocators()
 	{
-#if WITH_EDITOR
+#if WITH_CRUNCH_COMPRESSION
 		crn_set_memory_callbacks(&CrunchReallocFunc, &CrunchMSizeFunc, nullptr);
-#endif // WITH_EDITOR
+#endif // WITH_CRUNCH_COMPRESSION
 		crnd::crnd_set_memory_callbacks(&CrunchReallocFunc, &CrunchMSizeFunc, nullptr);
 	}
 };
