Disable stat tracking for Crunch memory used during compression
- Mostly care about runtime cost
- Compressing lots of textures generates can causes large overhead just processing stat messages (at least in debug build)

diff --git a/Engine/Source/Runtime/CrunchCompression/Private/CrunchCompression.cpp b/Engine/Source/Runtime/CrunchCompression/Private/CrunchCompression.cpp
--- a/Engine/Source/Runtime/CrunchCompression/Private/CrunchCompression.cpp
+++ b/Engine/Source/Runtime/CrunchCompression/Private/CrunchCompression.cpp
@@ -138,6 +138,7 @@ DECLARE_DWORD_ACCUMULATOR_STAT(TEXT("Total Allocations"), STAT_TotalAllocations,
 // Value exposed by crunch headers is inconsistent
 static const uint32 CRUNCH_MIN_ALLOC_ALIGNMENT = 2 * sizeof(SIZE_T);
 
+template<bool bEnableStats>
 static void* CrunchReallocFunc(void* p, size_t size, size_t* pActual_size, bool movable, void* pUser_data)
 {
 	void* Result = nullptr;
@@ -146,24 +147,36 @@ static void* CrunchReallocFunc(void* p, size_t size, size_t* pActual_size, bool
 	if (!p)
 	{
 		ensure(size > 0u);
-		INC_DWORD_STAT(STAT_TotalAllocations);
+		if (bEnableStats)
+		{
+			INC_DWORD_STAT(STAT_TotalAllocations);
+		}
 		Result = FMemory::Malloc(size, CRUNCH_MIN_ALLOC_ALIGNMENT);
 		ResultSize = FMemory::GetAllocSize(Result);
 	}
 	else if (size == 0u)
 	{
-		DEC_DWORD_STAT(STAT_TotalAllocations);
-		DEC_MEMORY_STAT_BY(STAT_TotalMemory, FMemory::GetAllocSize(p));
+		if (bEnableStats)
+		{
+			DEC_DWORD_STAT(STAT_TotalAllocations);
+			DEC_MEMORY_STAT_BY(STAT_TotalMemory, FMemory::GetAllocSize(p));
+		}
 		FMemory::Free(p);
 	}
 	else if (movable)
 	{
-		DEC_MEMORY_STAT_BY(STAT_TotalMemory, FMemory::GetAllocSize(p));
+		if (bEnableStats)
+		{
+			DEC_MEMORY_STAT_BY(STAT_TotalMemory, FMemory::GetAllocSize(p));
+		}
 		Result = FMemory::Realloc(p, size, CRUNCH_MIN_ALLOC_ALIGNMENT);
 		ResultSize = FMemory::GetAllocSize(Result);
 	}
 
-	INC_MEMORY_STAT_BY(STAT_TotalMemory, ResultSize);
+	if (bEnableStats)
+	{
+		INC_MEMORY_STAT_BY(STAT_TotalMemory, ResultSize);
+	}
 
 	if (pActual_size)
 	{
@@ -183,9 +196,10 @@ struct FCrunchRegisterAllocators
 	FCrunchRegisterAllocators()
 	{
 #if WITH_CRUNCH_COMPRESSION
-		crn_set_memory_callbacks(&CrunchReallocFunc, &CrunchMSizeFunc, nullptr);
+		// Don't track stats for Crunch memory used by compressor, only interested in memory used at runtime by decompression
+		crn_set_memory_callbacks(&CrunchReallocFunc<false>, &CrunchMSizeFunc, nullptr);
 #endif // WITH_CRUNCH_COMPRESSION
-		crnd::crnd_set_memory_callbacks(&CrunchReallocFunc, &CrunchMSizeFunc, nullptr);
+		crnd::crnd_set_memory_callbacks(&CrunchReallocFunc<true>, &CrunchMSizeFunc, nullptr);
 	}
 };
 static FCrunchRegisterAllocators gCrunchRegisterAllocators;
